- CIS fixes

diff --git a/Engine/Plugins/Runtime/IOSReplayKit/Source/IOSReplayKit/IOSReplayKit.Build.cs b/Engine/Plugins/Runtime/IOSReplayKit/Source/IOSReplayKit/IOSReplayKit.Build.cs
--- a/Engine/Plugins/Runtime/IOSReplayKit/Source/IOSReplayKit/IOSReplayKit.Build.cs
+++ b/Engine/Plugins/Runtime/IOSReplayKit/Source/IOSReplayKit/IOSReplayKit.Build.cs
@@ -1,4 +1,4 @@
-// Copyright 1998-2016 Epic Games, Inc. All Rights Reserved.
+// Copyright 1998-2019 Epic Games, Inc. All Rights Reserved.
 
 using System.IO;
 
diff --git a/Engine/Plugins/Runtime/IOSReplayKit/Source/IOSReplayKit/Private/IOSReplayKit.cpp b/Engine/Plugins/Runtime/IOSReplayKit/Source/IOSReplayKit/Private/IOSReplayKit.cpp
--- a/Engine/Plugins/Runtime/IOSReplayKit/Source/IOSReplayKit/Private/IOSReplayKit.cpp
+++ b/Engine/Plugins/Runtime/IOSReplayKit/Source/IOSReplayKit/Private/IOSReplayKit.cpp
@@ -1,4 +1,4 @@
-// Copyright 1998-2016 Epic Games, Inc. All Rights Reserved.
+// Copyright 1998-2019 Epic Games, Inc. All Rights Reserved.
 
 #include "IOSReplayKit.h"
 
diff --git a/Engine/Plugins/Runtime/IOSReplayKit/Source/IOSReplayKit/Private/IOSReplayKitControl.cpp b/Engine/Plugins/Runtime/IOSReplayKit/Source/IOSReplayKit/Private/IOSReplayKitControl.cpp
--- a/Engine/Plugins/Runtime/IOSReplayKit/Source/IOSReplayKit/Private/IOSReplayKitControl.cpp
+++ b/Engine/Plugins/Runtime/IOSReplayKit/Source/IOSReplayKit/Private/IOSReplayKitControl.cpp
@@ -1,4 +1,4 @@
-// Copyright 1998-2018 Epic Games, Inc. All Rights Reserved.
+// Copyright 1998-2019 Epic Games, Inc. All Rights Reserved.
 
 #include "IOSReplayKitControl.h"
 #include "IOSReplayKit.h"
diff --git a/Engine/Plugins/Runtime/IOSReplayKit/Source/IOSReplayKit/Private/ReplayKitRecorder.cpp b/Engine/Plugins/Runtime/IOSReplayKit/Source/IOSReplayKit/Private/ReplayKitRecorder.cpp
--- a/Engine/Plugins/Runtime/IOSReplayKit/Source/IOSReplayKit/Private/ReplayKitRecorder.cpp
+++ b/Engine/Plugins/Runtime/IOSReplayKit/Source/IOSReplayKit/Private/ReplayKitRecorder.cpp
@@ -1,4 +1,4 @@
-// Copyright 1998-2016 Epic Games, Inc. All Rights Reserved.
+// Copyright 1998-2019 Epic Games, Inc. All Rights Reserved.
 
 
 #include "ReplayKitRecorder.h"
diff --git a/Engine/Plugins/Runtime/IOSReplayKit/Source/IOSReplayKit/Private/ReplayKitRecorder.h b/Engine/Plugins/Runtime/IOSReplayKit/Source/IOSReplayKit/Private/ReplayKitRecorder.h
--- a/Engine/Plugins/Runtime/IOSReplayKit/Source/IOSReplayKit/Private/ReplayKitRecorder.h
+++ b/Engine/Plugins/Runtime/IOSReplayKit/Source/IOSReplayKit/Private/ReplayKitRecorder.h
@@ -1,4 +1,4 @@
-// Copyright 1998-2016 Epic Games, Inc. All Rights Reserved.
+// Copyright 1998-2019 Epic Games, Inc. All Rights Reserved.
 
 
 #include "CoreTypes.h"
diff --git a/Engine/Plugins/Runtime/IOSReplayKit/Source/IOSReplayKit/Public/IOSReplayKit.h b/Engine/Plugins/Runtime/IOSReplayKit/Source/IOSReplayKit/Public/IOSReplayKit.h
--- a/Engine/Plugins/Runtime/IOSReplayKit/Source/IOSReplayKit/Public/IOSReplayKit.h
+++ b/Engine/Plugins/Runtime/IOSReplayKit/Source/IOSReplayKit/Public/IOSReplayKit.h
@@ -1,4 +1,4 @@
-// Copyright 1998-2018 Epic Games, Inc. All Rights Reserved.
+// Copyright 1998-2019 Epic Games, Inc. All Rights Reserved.
 
 #pragma once
 
diff --git a/Engine/Plugins/Runtime/IOSReplayKit/Source/IOSReplayKit/Public/IOSReplayKitControl.h b/Engine/Plugins/Runtime/IOSReplayKit/Source/IOSReplayKit/Public/IOSReplayKitControl.h
--- a/Engine/Plugins/Runtime/IOSReplayKit/Source/IOSReplayKit/Public/IOSReplayKitControl.h
+++ b/Engine/Plugins/Runtime/IOSReplayKit/Source/IOSReplayKit/Public/IOSReplayKitControl.h
@@ -1,4 +1,4 @@
-// Copyright 1998-2018 Epic Games, Inc. All Rights Reserved.
+// Copyright 1998-2019 Epic Games, Inc. All Rights Reserved.
 
 #pragma once
 
diff --git a/Engine/Plugins/Runtime/WebBrowserNativeProxy/Source/WebBrowserNativeProxy/Private/WebBrowserNativeProxyModule.cpp b/Engine/Plugins/Runtime/WebBrowserNativeProxy/Source/WebBrowserNativeProxy/Private/WebBrowserNativeProxyModule.cpp
--- a/Engine/Plugins/Runtime/WebBrowserNativeProxy/Source/WebBrowserNativeProxy/Private/WebBrowserNativeProxyModule.cpp
+++ b/Engine/Plugins/Runtime/WebBrowserNativeProxy/Source/WebBrowserNativeProxy/Private/WebBrowserNativeProxyModule.cpp
@@ -1,4 +1,4 @@
-// Copyright 1998-2018 Epic Games, Inc. All Rights Reserved.
+// Copyright 1998-2019 Epic Games, Inc. All Rights Reserved.
 
 #include "WebBrowserNativeProxyModule.h"
 #include "Modules/ModuleManager.h"
diff --git a/Engine/Plugins/Runtime/WebBrowserNativeProxy/Source/WebBrowserNativeProxy/Public/WebBrowserNativeProxyModule.h b/Engine/Plugins/Runtime/WebBrowserNativeProxy/Source/WebBrowserNativeProxy/Public/WebBrowserNativeProxyModule.h
--- a/Engine/Plugins/Runtime/WebBrowserNativeProxy/Source/WebBrowserNativeProxy/Public/WebBrowserNativeProxyModule.h
+++ b/Engine/Plugins/Runtime/WebBrowserNativeProxy/Source/WebBrowserNativeProxy/Public/WebBrowserNativeProxyModule.h
@@ -1,4 +1,4 @@
-// Copyright 1998-2018 Epic Games, Inc. All Rights Reserved.
+// Copyright 1998-2019 Epic Games, Inc. All Rights Reserved.
 
 #pragma once
 
diff --git a/Engine/Plugins/Runtime/WebBrowserNativeProxy/Source/WebBrowserNativeProxy/WebBrowserNativeProxy.build.cs b/Engine/Plugins/Runtime/WebBrowserNativeProxy/Source/WebBrowserNativeProxy/WebBrowserNativeProxy.build.cs
--- a/Engine/Plugins/Runtime/WebBrowserNativeProxy/Source/WebBrowserNativeProxy/WebBrowserNativeProxy.build.cs
+++ b/Engine/Plugins/Runtime/WebBrowserNativeProxy/Source/WebBrowserNativeProxy/WebBrowserNativeProxy.build.cs
@@ -1,4 +1,4 @@
-// Copyright 1998-2018 Epic Games, Inc. All Rights Reserved.
+// Copyright 1998-2019 Epic Games, Inc. All Rights Reserved.
 
 namespace UnrealBuildTool.Rules
 {
diff --git a/Engine/Shaders/Private/GameplayMediaEncoderShaders.usf b/Engine/Shaders/Private/GameplayMediaEncoderShaders.usf
--- a/Engine/Shaders/Private/GameplayMediaEncoderShaders.usf
+++ b/Engine/Shaders/Private/GameplayMediaEncoderShaders.usf
@@ -1,4 +1,4 @@
-// Copyright 1998-2018 Epic Games, Inc. All Rights Reserved.
+// Copyright 1998-2019 Epic Games, Inc. All Rights Reserved.
 
 /*=========================================================================================
 	GameplayMediaEncoderShaders.usf: utility shaders for the GameplayMediaEncoder module
